UE-31727 Added Right Thumbstick input for Mobile

diff --git a/Templates/TP_Flying/Config/DefaultInput.ini b/Templates/TP_Flying/Config/DefaultInput.ini
--- a/Templates/TP_Flying/Config/DefaultInput.ini
+++ b/Templates/TP_Flying/Config/DefaultInput.ini
@@ -1,76 +1,68 @@
-
-[/Script/Engine.InputSettings]
 
+[/Script/Engine.InputSettings]
 -AxisConfig=(AxisKeyName="Gamepad_LeftX",AxisProperties=(DeadZone=0.25,Exponent=1.f,Sensitivity=1.f))
-
 -AxisConfig=(AxisKeyName="Gamepad_LeftY",AxisProperties=(DeadZone=0.25,Exponent=1.f,Sensitivity=1.f))
-
 -AxisConfig=(AxisKeyName="Gamepad_RightX",AxisProperties=(DeadZone=0.25,Exponent=1.f,Sensitivity=1.f))
-
 -AxisConfig=(AxisKeyName="Gamepad_RightY",AxisProperties=(DeadZone=0.25,Exponent=1.f,Sensitivity=1.f))
-
 -AxisConfig=(AxisKeyName="MouseX",AxisProperties=(DeadZone=0.f,Exponent=1.f,Sensitivity=0.07f))
-
 -AxisConfig=(AxisKeyName="MouseY",AxisProperties=(DeadZone=0.f,Exponent=1.f,Sensitivity=0.07f))
-
 +AxisConfig=(AxisKeyName="Gamepad_LeftX",AxisProperties=(DeadZone=0.250000,Sensitivity=1.000000,Exponent=1.000000,bInvert=False))
-
 +AxisConfig=(AxisKeyName="Gamepad_LeftY",AxisProperties=(DeadZone=0.250000,Sensitivity=1.000000,Exponent=1.000000,bInvert=False))
-
 +AxisConfig=(AxisKeyName="Gamepad_RightX",AxisProperties=(DeadZone=0.250000,Sensitivity=1.000000,Exponent=1.000000,bInvert=False))
-
 +AxisConfig=(AxisKeyName="Gamepad_RightY",AxisProperties=(DeadZone=0.250000,Sensitivity=1.000000,Exponent=1.000000,bInvert=False))
-
 +AxisConfig=(AxisKeyName="MouseX",AxisProperties=(DeadZone=0.000000,Sensitivity=0.070000,Exponent=1.000000,bInvert=False))
-
 +AxisConfig=(AxisKeyName="MouseY",AxisProperties=(DeadZone=0.000000,Sensitivity=0.070000,Exponent=1.000000,bInvert=False))
-
++AxisConfig=(AxisKeyName="MouseWheelAxis",AxisProperties=(DeadZone=0.000000,Sensitivity=1.000000,Exponent=1.000000,bInvert=False))
++AxisConfig=(AxisKeyName="Gamepad_LeftTriggerAxis",AxisProperties=(DeadZone=0.000000,Sensitivity=1.000000,Exponent=1.000000,bInvert=False))
++AxisConfig=(AxisKeyName="Gamepad_RightTriggerAxis",AxisProperties=(DeadZone=0.000000,Sensitivity=1.000000,Exponent=1.000000,bInvert=False))
++AxisConfig=(AxisKeyName="MotionController_Left_Thumbstick_X",AxisProperties=(DeadZone=0.000000,Sensitivity=1.000000,Exponent=1.000000,bInvert=False))
++AxisConfig=(AxisKeyName="MotionController_Left_Thumbstick_Y",AxisProperties=(DeadZone=0.000000,Sensitivity=1.000000,Exponent=1.000000,bInvert=False))
++AxisConfig=(AxisKeyName="MotionController_Left_TriggerAxis",AxisProperties=(DeadZone=0.000000,Sensitivity=1.000000,Exponent=1.000000,bInvert=False))
++AxisConfig=(AxisKeyName="MotionController_Left_Grip1Axis",AxisProperties=(DeadZone=0.000000,Sensitivity=1.000000,Exponent=1.000000,bInvert=False))
++AxisConfig=(AxisKeyName="MotionController_Left_Grip2Axis",AxisProperties=(DeadZone=0.000000,Sensitivity=1.000000,Exponent=1.000000,bInvert=False))
++AxisConfig=(AxisKeyName="MotionController_Right_Thumbstick_X",AxisProperties=(DeadZone=0.000000,Sensitivity=1.000000,Exponent=1.000000,bInvert=False))
++AxisConfig=(AxisKeyName="MotionController_Right_Thumbstick_Y",AxisProperties=(DeadZone=0.000000,Sensitivity=1.000000,Exponent=1.000000,bInvert=False))
++AxisConfig=(AxisKeyName="MotionController_Right_TriggerAxis",AxisProperties=(DeadZone=0.000000,Sensitivity=1.000000,Exponent=1.000000,bInvert=False))
++AxisConfig=(AxisKeyName="MotionController_Right_Grip1Axis",AxisProperties=(DeadZone=0.000000,Sensitivity=1.000000,Exponent=1.000000,bInvert=False))
++AxisConfig=(AxisKeyName="MotionController_Right_Grip2Axis",AxisProperties=(DeadZone=0.000000,Sensitivity=1.000000,Exponent=1.000000,bInvert=False))
++AxisConfig=(AxisKeyName="Gamepad_Special_Left_X",AxisProperties=(DeadZone=0.000000,Sensitivity=1.000000,Exponent=1.000000,bInvert=False))
++AxisConfig=(AxisKeyName="Gamepad_Special_Left_Y",AxisProperties=(DeadZone=0.000000,Sensitivity=1.000000,Exponent=1.000000,bInvert=False))
+bAltEnterTogglesFullscreen=True
 bUseMouseForTouch=False
-
+bEnableMouseSmoothing=True
+bEnableFOVScaling=True
+FOVScale=0.011110
 DoubleClickTime=0.200000
-
+bCaptureMouseOnLaunch=True
+DefaultViewportMouseCaptureMode=CapturePermanently_IncludingInitialMouseDown
+bDefaultViewportMouseLock=True
 +AxisMappings=(AxisName="MoveUp",Key=W,Scale=1.000000)
-
 +AxisMappings=(AxisName="MoveUp",Key=S,Scale=-1.000000)
-
 +AxisMappings=(AxisName="MoveUp",Key=Gamepad_LeftStick_Up,Scale=1.000000)
-
 +AxisMappings=(AxisName="MoveUp",Key=Gamepad_LeftStick_Down,Scale=-1.000000)
-
 +AxisMappings=(AxisName="MoveUp",Key=Up,Scale=1.000000)
-
 +AxisMappings=(AxisName="MoveUp",Key=Down,Scale=-1.000000)
-
 +AxisMappings=(AxisName="MoveRight",Key=D,Scale=1.000000)
-
 +AxisMappings=(AxisName="MoveRight",Key=A,Scale=-1.000000)
-
 +AxisMappings=(AxisName="MoveRight",Key=Gamepad_LeftStick_Right,Scale=1.000000)
-
 +AxisMappings=(AxisName="MoveRight",Key=Gamepad_LeftStick_Left,Scale=-1.000000)
-
 +AxisMappings=(AxisName="MoveRight",Key=Right,Scale=1.000000)
-
 +AxisMappings=(AxisName="MoveRight",Key=Left,Scale=-1.000000)
-
 +AxisMappings=(AxisName="Thrust",Key=LeftShift,Scale=1.000000)
-
 +AxisMappings=(AxisName="Thrust",Key=RightShift,Scale=1.000000)
-
-+AxisMappings=(AxisName="Thrust",Key=Gamepad_RightTriggerAxis,Scale=1.000000)
-
-+AxisMappings=(AxisName="Thrust",Key=Gamepad_LeftTriggerAxis,Scale=-1.000000)
-
++AxisMappings=(AxisName="Thrust",Key=Gamepad_RightY,Scale=1.000000)
++AxisMappings=(AxisName="Thrust",Key=Gamepad_RightY,Scale=-1.000000)
 +AxisMappings=(AxisName="Thrust",Key=LeftControl,Scale=-1.000000)
-
 +AxisMappings=(AxisName="Thrust",Key=RightControl,Scale=-1.000000)
-
 +AxisMappings=(AxisName="MoveUp",Key=Gamepad_LeftY,Scale=1.000000)
-
 +AxisMappings=(AxisName="MoveRight",Key=Gamepad_LeftX,Scale=1.000000)
-
 +AxisMappings=(AxisName="Thrust",Key=Gamepad_RightY,Scale=1.000000)
-
 +AxisMappings=(AxisName="Thrust",Key=SpaceBar,Scale=1.000000)
-DefaultTouchInterface=/Engine/MobileResources/HUD/LeftVirtualJoystickOnly.LeftVirtualJoystickOnly
+bAlwaysShowTouchInterface=False
+bShowConsoleOnFourFingerTap=True
+DefaultTouchInterface=/Engine/MobileResources/HUD/DefaultVirtualJoysticks.DefaultVirtualJoysticks
+ConsoleKey=None
+-ConsoleKeys=Tilde
++ConsoleKeys=Tilde
+
 
diff --git a/Templates/TP_FlyingBP/Config/DefaultInput.ini b/Templates/TP_FlyingBP/Config/DefaultInput.ini
--- a/Templates/TP_FlyingBP/Config/DefaultInput.ini
+++ b/Templates/TP_FlyingBP/Config/DefaultInput.ini
@@ -27,6 +27,18 @@
 +AxisConfig=(AxisKeyName="MotionController_Right_Grip2Axis",AxisProperties=(DeadZone=0.000000,Sensitivity=1.000000,Exponent=1.000000,bInvert=False))
 +AxisConfig=(AxisKeyName="Gamepad_Special_Left_X",AxisProperties=(DeadZone=0.000000,Sensitivity=1.000000,Exponent=1.000000,bInvert=False))
 +AxisConfig=(AxisKeyName="Gamepad_Special_Left_Y",AxisProperties=(DeadZone=0.000000,Sensitivity=1.000000,Exponent=1.000000,bInvert=False))
++AxisConfig=(AxisKeyName="OculusTouch_Left_Thumbstick",AxisProperties=(DeadZone=0.000000,Sensitivity=1.000000,Exponent=1.000000,bInvert=False))
++AxisConfig=(AxisKeyName="OculusTouch_Left_FaceButton1",AxisProperties=(DeadZone=0.000000,Sensitivity=1.000000,Exponent=1.000000,bInvert=False))
++AxisConfig=(AxisKeyName="OculusTouch_Left_Trigger",AxisProperties=(DeadZone=0.000000,Sensitivity=1.000000,Exponent=1.000000,bInvert=False))
++AxisConfig=(AxisKeyName="OculusTouch_Left_FaceButton2",AxisProperties=(DeadZone=0.000000,Sensitivity=1.000000,Exponent=1.000000,bInvert=False))
++AxisConfig=(AxisKeyName="OculusTouch_Left_IndexPointing",AxisProperties=(DeadZone=0.000000,Sensitivity=1.000000,Exponent=1.000000,bInvert=False))
++AxisConfig=(AxisKeyName="OculusTouch_Left_ThumbUp",AxisProperties=(DeadZone=0.000000,Sensitivity=1.000000,Exponent=1.000000,bInvert=False))
++AxisConfig=(AxisKeyName="OculusTouch_Right_Thumbstick",AxisProperties=(DeadZone=0.000000,Sensitivity=1.000000,Exponent=1.000000,bInvert=False))
++AxisConfig=(AxisKeyName="OculusTouch_Right_FaceButton1",AxisProperties=(DeadZone=0.000000,Sensitivity=1.000000,Exponent=1.000000,bInvert=False))
++AxisConfig=(AxisKeyName="OculusTouch_Right_Trigger",AxisProperties=(DeadZone=0.000000,Sensitivity=1.000000,Exponent=1.000000,bInvert=False))
++AxisConfig=(AxisKeyName="OculusTouch_Right_FaceButton2",AxisProperties=(DeadZone=0.000000,Sensitivity=1.000000,Exponent=1.000000,bInvert=False))
++AxisConfig=(AxisKeyName="OculusTouch_Right_IndexPointing",AxisProperties=(DeadZone=0.000000,Sensitivity=1.000000,Exponent=1.000000,bInvert=False))
++AxisConfig=(AxisKeyName="OculusTouch_Right_ThumbUp",AxisProperties=(DeadZone=0.000000,Sensitivity=1.000000,Exponent=1.000000,bInvert=False))
 bAltEnterTogglesFullscreen=True
 bUseMouseForTouch=False
 bEnableMouseSmoothing=True
@@ -50,8 +62,8 @@ bDefaultViewportMouseLock=True
 +AxisMappings=(AxisName="MoveRight",Key=Left,Scale=-1.000000)
 +AxisMappings=(AxisName="Thrust",Key=LeftShift,Scale=1.000000)
 +AxisMappings=(AxisName="Thrust",Key=RightShift,Scale=1.000000)
-+AxisMappings=(AxisName="Thrust",Key=Gamepad_RightTriggerAxis,Scale=1.000000)
-+AxisMappings=(AxisName="Thrust",Key=Gamepad_LeftTriggerAxis,Scale=-1.000000)
++AxisMappings=(AxisName="Thrust",Key=Gamepad_RightY,Scale=1.000000)
++AxisMappings=(AxisName="Thrust",Key=Gamepad_RightY,Scale=-1.000000)
 +AxisMappings=(AxisName="Thrust",Key=LeftControl,Scale=-1.000000)
 +AxisMappings=(AxisName="Thrust",Key=RightControl,Scale=-1.000000)
 +AxisMappings=(AxisName="MoveUp",Key=Gamepad_LeftY,Scale=1.000000)
@@ -66,7 +78,7 @@ bDefaultViewportMouseLock=True
 +AxisMappings=(AxisName="Thrust",Key=MotionController_Left_Trigger,Scale=-1.000000)
 bAlwaysShowTouchInterface=False
 bShowConsoleOnFourFingerTap=True
-DefaultTouchInterface=/Engine/MobileResources/HUD/LeftVirtualJoystickOnly.LeftVirtualJoystickOnly
+DefaultTouchInterface=/Engine/MobileResources/HUD/DefaultVirtualJoysticks.DefaultVirtualJoysticks
 ConsoleKey=None
 -ConsoleKeys=Tilde
 +ConsoleKeys=Tilde
